#UE4 Compile fix for slate null renderer.

diff --git a/Engine/Source/Runtime/SlateNullRenderer/Private/SlateNullRenderer.cpp b/Engine/Source/Runtime/SlateNullRenderer/Private/SlateNullRenderer.cpp
--- a/Engine/Source/Runtime/SlateNullRenderer/Private/SlateNullRenderer.cpp
+++ b/Engine/Source/Runtime/SlateNullRenderer/Private/SlateNullRenderer.cpp
@@ -52,6 +52,11 @@ bool FSlateNullRenderer::GenerateDynamicImageResource( FName ResourceName, uint3
 	return false;
 }
 
+FSlateResourceHandle FSlateNullRenderer::GetResourceHandle( const FSlateBrush& Brush )
+{
+	return FSlateResourceHandle();
+}
+
 void FSlateNullRenderer::RemoveDynamicBrushResource( TSharedPtr<FSlateDynamicImageBrush> BrushToRemove )
 {
 }
diff --git a/Engine/Source/Runtime/SlateNullRenderer/Private/SlateNullRenderer.h b/Engine/Source/Runtime/SlateNullRenderer/Private/SlateNullRenderer.h
--- a/Engine/Source/Runtime/SlateNullRenderer/Private/SlateNullRenderer.h
+++ b/Engine/Source/Runtime/SlateNullRenderer/Private/SlateNullRenderer.h
@@ -22,6 +22,7 @@ public:
 	virtual void RemoveDynamicBrushResource( TSharedPtr<FSlateDynamicImageBrush> BrushToRemove ) override;
 	virtual FIntPoint GenerateDynamicImageResource(const FName InTextureName) override;
 	virtual bool GenerateDynamicImageResource( FName ResourceName, uint32 Width, uint32 Height, const TArray< uint8 >& Bytes ) override;
+	virtual FSlateResourceHandle GetResourceHandle( const FSlateBrush& Brush ) override;
 	virtual FSlateUpdatableTexture* CreateUpdatableTexture(uint32 Width, uint32 Height) override;
 	virtual void ReleaseUpdatableTexture(FSlateUpdatableTexture* Texture) override;
 };
