Fix non-unity and mac config

diff --git a/Engine/Source/Runtime/SlateCore/Public/Rendering/SlateDrawBuffer.h b/Engine/Source/Runtime/SlateCore/Public/Rendering/SlateDrawBuffer.h
--- a/Engine/Source/Runtime/SlateCore/Public/Rendering/SlateDrawBuffer.h
+++ b/Engine/Source/Runtime/SlateCore/Public/Rendering/SlateDrawBuffer.h
@@ -3,6 +3,7 @@
 #pragma once
 
 #include "CoreMinimal.h"
+#include "UObject/GCObject.h"
 
 class FSlateWindowElementList;
 class SWindow;
diff --git a/Engine/Source/Runtime/SlateNullRenderer/Private/SlateNullRenderer.cpp b/Engine/Source/Runtime/SlateNullRenderer/Private/SlateNullRenderer.cpp
--- a/Engine/Source/Runtime/SlateNullRenderer/Private/SlateNullRenderer.cpp
+++ b/Engine/Source/Runtime/SlateNullRenderer/Private/SlateNullRenderer.cpp
@@ -20,9 +20,14 @@ void FSlateNullRenderer::Destroy()
 
 FSlateDrawBuffer& FSlateNullRenderer::GetDrawBuffer()
 {
-	static FSlateDrawBuffer StaticDrawBuffer;
-	StaticDrawBuffer.ClearBuffer();
-	return StaticDrawBuffer;
+	static TUniquePtr<FSlateDrawBuffer> StaticDrawBuffer;
+	if (!StaticDrawBuffer.IsValid())
+	{
+		StaticDrawBuffer = MakeUnique<FSlateDrawBuffer>();
+	}
+
+	StaticDrawBuffer->ClearBuffer();
+	return *StaticDrawBuffer;
 }
 
 void FSlateNullRenderer::CreateViewport( const TSharedRef<SWindow> Window )
