Added missing dependency on CoreUObject

diff --git a/Engine/Source/Runtime/SlateNullRenderer/SlateNullRenderer.Build.cs b/Engine/Source/Runtime/SlateNullRenderer/SlateNullRenderer.Build.cs
--- a/Engine/Source/Runtime/SlateNullRenderer/SlateNullRenderer.Build.cs
+++ b/Engine/Source/Runtime/SlateNullRenderer/SlateNullRenderer.Build.cs
@@ -11,6 +11,7 @@ public class SlateNullRenderer : ModuleRules
 		PrivateDependencyModuleNames.AddRange(
 			new string[] {
 				"Core",
+				"CoreUObject",
 				"SlateCore"
 			}
 		);
