Fixed compile error when a program depends on ApplicationCore without the editor enabled.

diff --git a/Engine/Source/Runtime/SlateNullRenderer/Private/SlateNullRendererModule.cpp b/Engine/Source/Runtime/SlateNullRenderer/Private/SlateNullRendererModule.cpp
--- a/Engine/Source/Runtime/SlateNullRenderer/Private/SlateNullRendererModule.cpp
+++ b/Engine/Source/Runtime/SlateNullRenderer/Private/SlateNullRendererModule.cpp
@@ -9,7 +9,9 @@
 #include "Fonts/FontCache.h"
 #include "Rendering/DrawElements.h"
 #include "Rendering/SlateRenderer.h"
+#if WITH_ENGINE
 #include "TextureResource.h"
+#endif
 #include "Interfaces/ISlateNullRendererModule.h"
 #include "SlateNullRenderer.h"
 
@@ -25,7 +27,7 @@ public:
 #endif
 
 	// FSlateShaderResourceManager interface
-	virtual FSlateShaderResourceProxy* GetShaderResource( const FSlateBrush& InBrush ) override { return nullptr; }
+	virtual FSlateShaderResourceProxy* GetShaderResource(const FSlateBrush& Brush, FVector2D LocalSize, float DrawScale) override { return nullptr; }
 	virtual ISlateAtlasProvider* GetTextureAtlasProvider() { return this; }
 	virtual FSlateResourceHandle GetResourceHandle( const FSlateBrush& InBrush ) override 
 	{
@@ -67,7 +69,7 @@ public:
 		: FSlateFontAtlas(AtlasSize, AtlasSize, true)
 	{}
 
-	virtual class FSlateShaderResource* GetSlateTexture() override { return &NullFontTexture; }
+	virtual class FSlateShaderResource* GetSlateTexture() const override { return &NullFontTexture; }
 	virtual class FTextureResource* GetEngineTexture() override
 	{
 #if WITH_ENGINE
