Speculative fix Mac NonUnity

diff --git a/Engine/Source/Runtime/SlateNullRenderer/Private/SlateNullRendererModule.cpp b/Engine/Source/Runtime/SlateNullRenderer/Private/SlateNullRendererModule.cpp
--- a/Engine/Source/Runtime/SlateNullRenderer/Private/SlateNullRendererModule.cpp
+++ b/Engine/Source/Runtime/SlateNullRenderer/Private/SlateNullRendererModule.cpp
@@ -8,6 +8,7 @@
 #include "Fonts/FontTypes.h"
 #include "Fonts/FontCache.h"
 #include "Rendering/DrawElements.h"
+#include "Rendering/SlateDrawBuffer.h"
 #include "Rendering/SlateRenderer.h"
 #if WITH_ENGINE
 #include "TextureResource.h"
