Fixed a leak when using the Null renderer, by implementing a valid Sync method. Without a Sync, the main thread will overflow the render thread with an ever growing list of render commands.

diff --git a/Engine/Source/Runtime/SlateNullRenderer/Private/SlateNullRenderer.cpp b/Engine/Source/Runtime/SlateNullRenderer/Private/SlateNullRenderer.cpp
--- a/Engine/Source/Runtime/SlateNullRenderer/Private/SlateNullRenderer.cpp
+++ b/Engine/Source/Runtime/SlateNullRenderer/Private/SlateNullRenderer.cpp
@@ -130,4 +130,12 @@ int32 FSlateNullRenderer::GetCurrentSceneIndex() const
 void FSlateNullRenderer::ClearScenes() 
 {
 	// This is a no-op
-}
\ No newline at end of file
+}
+
+void FSlateNullRenderer::Sync() const
+{
+	// Sync game and render thread. Either total sync or allowing one frame lag.
+	static FFrameEndSync FrameEndSync;
+	static auto CVarAllowOneFrameThreadLag = IConsoleManager::Get().FindTConsoleVariableDataInt(TEXT("r.OneFrameThreadLag"));
+	FrameEndSync.Sync(CVarAllowOneFrameThreadLag->GetValueOnAnyThread() != 0);
+}
diff --git a/Engine/Source/Runtime/SlateNullRenderer/Private/SlateNullRenderer.h b/Engine/Source/Runtime/SlateNullRenderer/Private/SlateNullRenderer.h
--- a/Engine/Source/Runtime/SlateNullRenderer/Private/SlateNullRenderer.h
+++ b/Engine/Source/Runtime/SlateNullRenderer/Private/SlateNullRenderer.h
@@ -42,6 +42,7 @@ public:
 	virtual int32 RegisterCurrentScene(FSceneInterface* Scene) override;
 	virtual int32 GetCurrentSceneIndex() const override;
 	virtual void ClearScenes() override;
+	virtual void Sync() const override;
 
 private:
 	TSharedPtr<FSlateShaderResourceManager> ResourceManager;
