fix missing include

diff --git a/Engine/Source/Runtime/SlateNullRenderer/Private/SlateNullRenderer.cpp b/Engine/Source/Runtime/SlateNullRenderer/Private/SlateNullRenderer.cpp
--- a/Engine/Source/Runtime/SlateNullRenderer/Private/SlateNullRenderer.cpp
+++ b/Engine/Source/Runtime/SlateNullRenderer/Private/SlateNullRenderer.cpp
@@ -3,6 +3,7 @@
 #include "SlateNullRenderer.h"
 #include "Rendering/DrawElements.h"
 #include "Rendering/SlateDrawBuffer.h"
+#include "UnrealEngine.h"
 
 
 FSlateNullRenderer::FSlateNullRenderer(const TSharedRef<FSlateFontServices>& InSlateFontServices, const TSharedRef<FSlateShaderResourceManager>& InResourceManager)
